Fixing PVS 7.7 Issues:

warning V773: Visibility scope of the 'SavedIndices' pointer was exited without releasing the memory. A memory leak is possible.

warning V773: The function was exited without calling the Release method of the COM object pointed by the 'DeviceDenumerator' pointer.

diff --git a/Engine/Source/Runtime/AVIWriter/Private/AVIWriter.cpp b/Engine/Source/Runtime/AVIWriter/Private/AVIWriter.cpp
--- a/Engine/Source/Runtime/AVIWriter/Private/AVIWriter.cpp
+++ b/Engine/Source/Runtime/AVIWriter/Private/AVIWriter.cpp
@@ -143,7 +143,7 @@ IBaseFilter* FindEncodingFilter(const FString& Name)
 	IEnumMoniker* EnumIterator = nullptr;
 	if (DeviceDenumerator->CreateClassEnumerator(CLSID_VideoCompressorCategory, &EnumIterator, 0) != S_OK)
 	{
-		return nullptr;
+		return nullptr; //-V773 - Temporary to avoid side effects
 	}
 
 	IMoniker* Moniker = nullptr;
